Work around import directives causing ICEs in VS2017 15.9.4 when preprocessing source files for PVS Studio.

diff --git a/Engine/Plugins/Developer/VisualStudioSourceCodeAccess/Source/VisualStudioSourceCodeAccess/VisualStudioSourceCodeAccess.Build.cs b/Engine/Plugins/Developer/VisualStudioSourceCodeAccess/Source/VisualStudioSourceCodeAccess/VisualStudioSourceCodeAccess.Build.cs
--- a/Engine/Plugins/Developer/VisualStudioSourceCodeAccess/Source/VisualStudioSourceCodeAccess/VisualStudioSourceCodeAccess.Build.cs
+++ b/Engine/Plugins/Developer/VisualStudioSourceCodeAccess/Source/VisualStudioSourceCodeAccess/VisualStudioSourceCodeAccess.Build.cs
@@ -58,7 +58,7 @@ namespace UnrealBuildTool.Rules
 				}
 			}
 
-			if (bHasVisualStudioDTE)
+			if (bHasVisualStudioDTE && Target.WindowsPlatform.StaticAnalyzer != WindowsStaticAnalyzer.PVSStudio)
 			{
 				PublicDefinitions.Add("VSACCESSOR_HAS_DTE=1");
 			}
diff --git a/Engine/Plugins/Runtime/Database/ADOSupport/Source/ADOSupport/Private/ADOSupport.cpp b/Engine/Plugins/Runtime/Database/ADOSupport/Source/ADOSupport/Private/ADOSupport.cpp
--- a/Engine/Plugins/Runtime/Database/ADOSupport/Source/ADOSupport/Private/ADOSupport.cpp
+++ b/Engine/Plugins/Runtime/Database/ADOSupport/Source/ADOSupport/Private/ADOSupport.cpp
@@ -5,8 +5,13 @@
 
 DEFINE_LOG_CATEGORY_STATIC(LogADODataBase, Log, All);
 
+#ifdef PVS_STUDIO
+// Bug in the MSVC preprocessor (as of VS2017 15.9.4) causes ICE when preprocessing files with import directives
+#define USE_ADO_INTEGRATION 0
+#else
 // @todo clang: #import is not supported by Clang on Windows platform, but we currently need this for ADO symbol importing.  For now we disable ADO support in Clang builds.
 #define USE_ADO_INTEGRATION (!PLATFORM_COMPILER_CLANG)
+#endif
 
 #if USE_ADO_INTEGRATION
 #include "Windows/WindowsHWrapper.h"
