Compilation fix for VS2019 v16.5.0

diff --git a/Engine/Plugins/Runtime/Database/ADOSupport/Source/ADOSupport/Private/ADOSupport.cpp b/Engine/Plugins/Runtime/Database/ADOSupport/Source/ADOSupport/Private/ADOSupport.cpp
--- a/Engine/Plugins/Runtime/Database/ADOSupport/Source/ADOSupport/Private/ADOSupport.cpp
+++ b/Engine/Plugins/Runtime/Database/ADOSupport/Source/ADOSupport/Private/ADOSupport.cpp
@@ -252,7 +252,7 @@ public:
 	/** Destructor, cleaning up ADO record set. */
 	virtual ~FADODataBaseRecordSet()
 	{
-		if(ADORecordSet && (ADORecordSet->State & ADODB::adStateOpen))
+		if(ADORecordSet != nullptr && (ADORecordSet->State & ADODB::adStateOpen))
 		{
 			// We're using smart pointers so all we need to do is close and assign NULL.
 			ADORecordSet->Close();
@@ -325,7 +325,7 @@ public:
 	virtual void Close()
 	{
 		// Close database connection if exists and free smart pointer.
-		if( DataBaseConnection && (DataBaseConnection->State & ADODB::adStateOpen))
+		if( DataBaseConnection != nullptr && (DataBaseConnection->State & ADODB::adStateOpen))
 		{
 			DataBaseConnection->Close();
 
