Fix intermittent errors on build farm caused by type library header having an invalid timestamp.

diff --git a/Engine/Plugins/Runtime/Database/ADOSupport/Source/ADOSupport/ADOSupport.Build.cs b/Engine/Plugins/Runtime/Database/ADOSupport/Source/ADOSupport/ADOSupport.Build.cs
--- a/Engine/Plugins/Runtime/Database/ADOSupport/Source/ADOSupport/ADOSupport.Build.cs
+++ b/Engine/Plugins/Runtime/Database/ADOSupport/Source/ADOSupport/ADOSupport.Build.cs
@@ -1,5 +1,8 @@
 // Copyright Epic Games, Inc. All Rights Reserved.
 
+using System;
+using System.IO;
+
 namespace UnrealBuildTool.Rules
 {
 	public class ADOSupport: ModuleRules
@@ -7,7 +10,9 @@ namespace UnrealBuildTool.Rules
 		public ADOSupport(ReadOnlyTargetRules Target) : base(Target)
 		{
 			AddEngineThirdPartyPrivateStaticDependencies(Target, "ADO");
-			bBuildLocallyWithSNDBS = true;
+
+			string MsAdo15 = Path.Combine(Environment.GetFolderPath(Environment.SpecialFolder.CommonProgramFiles), "System", "ADO", "msado15.dll");
+			TypeLibraries.Add(new TypeLibrary(MsAdo15, "rename(\"EOF\", \"ADOEOF\")", "msado15.tlh"));
 			
 			PublicIncludePaths.AddRange(
 				new string[] {
diff --git a/Engine/Plugins/Runtime/Database/ADOSupport/Source/ADOSupport/Private/ADOSupport.cpp b/Engine/Plugins/Runtime/Database/ADOSupport/Source/ADOSupport/Private/ADOSupport.cpp
--- a/Engine/Plugins/Runtime/Database/ADOSupport/Source/ADOSupport/Private/ADOSupport.cpp
+++ b/Engine/Plugins/Runtime/Database/ADOSupport/Source/ADOSupport/Private/ADOSupport.cpp
@@ -28,7 +28,7 @@ DEFINE_LOG_CATEGORY_STATIC(LogADODataBase, Log, All);
 
 #pragma warning(push)
 #pragma warning(disable: 4471) // a forward declaration of an unscoped enumeration must have an underlying type (int assumed)
-#import "System\ADO\msado15.dll" rename("EOF", "ADOEOF") //lint !e322
+#include "msado15.tlh"
 #pragma warning(pop)
 
 /*-----------------------------------------------------------------------------
