Disable ADO support on Clang.

diff --git a/Engine/Plugins/Runtime/Database/ADOSupport/Source/ADOSupport/ADOSupport.Build.cs b/Engine/Plugins/Runtime/Database/ADOSupport/Source/ADOSupport/ADOSupport.Build.cs
--- a/Engine/Plugins/Runtime/Database/ADOSupport/Source/ADOSupport/ADOSupport.Build.cs
+++ b/Engine/Plugins/Runtime/Database/ADOSupport/Source/ADOSupport/ADOSupport.Build.cs
@@ -11,8 +11,11 @@ namespace UnrealBuildTool.Rules
 		{
 			AddEngineThirdPartyPrivateStaticDependencies(Target, "ADO");
 
-			string MsAdo15 = Path.Combine(Environment.GetFolderPath(Environment.SpecialFolder.CommonProgramFiles), "System", "ADO", "msado15.dll");
-			TypeLibraries.Add(new TypeLibrary(MsAdo15, "rename(\"EOF\", \"ADOEOF\")", "msado15.tlh"));
+			if (Target.WindowsPlatform.Compiler != WindowsCompiler.Clang)
+			{
+				string MsAdo15 = Path.Combine(Environment.GetFolderPath(Environment.SpecialFolder.CommonProgramFiles), "System", "ADO", "msado15.dll");
+				TypeLibraries.Add(new TypeLibrary(MsAdo15, "rename(\"EOF\", \"ADOEOF\")", "msado15.tlh"));
+			}
 			
 			PublicIncludePaths.AddRange(
 				new string[] {
