Fix static analysis error due to type library being enabled.

diff --git a/Engine/Plugins/Runtime/Database/ADOSupport/Source/ADOSupport/ADOSupport.Build.cs b/Engine/Plugins/Runtime/Database/ADOSupport/Source/ADOSupport/ADOSupport.Build.cs
--- a/Engine/Plugins/Runtime/Database/ADOSupport/Source/ADOSupport/ADOSupport.Build.cs
+++ b/Engine/Plugins/Runtime/Database/ADOSupport/Source/ADOSupport/ADOSupport.Build.cs
@@ -11,7 +11,9 @@ namespace UnrealBuildTool.Rules
 		{
 			AddEngineThirdPartyPrivateStaticDependencies(Target, "ADO");
 
-			if (Target.WindowsPlatform.Compiler != WindowsCompiler.Clang)
+			if (Target.Platform == UnrealTargetPlatform.Win64 &&
+				Target.WindowsPlatform.Compiler != WindowsCompiler.Clang &&
+				Target.WindowsPlatform.StaticAnalyzer != WindowsStaticAnalyzer.PVSStudio)
 			{
 				string MsAdo15 = Path.Combine(Environment.GetFolderPath(Environment.SpecialFolder.CommonProgramFiles), "System", "ADO", "msado15.dll");
 				TypeLibraries.Add(new TypeLibrary(MsAdo15, "rename(\"EOF\", \"ADOEOF\")", "msado15.tlh"));
