ADOSupport: Don't throw an exception if exceptions aren't enabled

#rnx

diff --git a/Engine/Plugins/Runtime/Database/ADOSupport/Source/ADOSupport/Private/ADOSupport.cpp b/Engine/Plugins/Runtime/Database/ADOSupport/Source/ADOSupport/Private/ADOSupport.cpp
--- a/Engine/Plugins/Runtime/Database/ADOSupport/Source/ADOSupport/Private/ADOSupport.cpp
+++ b/Engine/Plugins/Runtime/Database/ADOSupport/Source/ADOSupport/Private/ADOSupport.cpp
@@ -301,7 +301,11 @@ public:
 		if (FAILED(hr))
 		{
 			FWindowsPlatformMisc::CoUninitialize();
+#if defined(PLATFORM_EXCEPTIONS_DISABLED) && PLATFORM_EXCEPTIONS_DISABLED == 0
 			throw _com_error(hr);
+#else
+			return false;
+#endif 
 		}
 
 		// Open the connection. Operation is synchronous.
